ReplicationSystemTest now has an include dependency on DerivedDataCache instead of a module dependency.

diff --git a/Engine/Source/Programs/ReplicationSystemTest/ReplicationSystemTest.Build.cs b/Engine/Source/Programs/ReplicationSystemTest/ReplicationSystemTest.Build.cs
--- a/Engine/Source/Programs/ReplicationSystemTest/ReplicationSystemTest.Build.cs
+++ b/Engine/Source/Programs/ReplicationSystemTest/ReplicationSystemTest.Build.cs
@@ -11,6 +11,13 @@ public class ReplicationSystemTest : ModuleRules
 		// For LaunchEngineLoop.cpp include
 		PrivateIncludePaths.Add("Runtime/Launch/Private");
 
+		PrivateIncludePathModuleNames.AddRange(
+			new string[]
+			{
+				"DerivedDataCache",
+			}
+		);
+
 		PrivateDependencyModuleNames.AddRange(
 			new string[]
 			{
@@ -19,7 +26,6 @@ public class ReplicationSystemTest : ModuleRules
  				"AutomationWorker",
                 "Core",
 				"Projects",
-				"DerivedDataCache",
 				"Engine",
 				"HeadMountedDisplay",
 				"InstallBundleManager",
