Get ReplicationSystemTest functional again.

#rnx

diff --git a/Engine/Source/Programs/ReplicationSystemTest/Private/ReplicationSystemTest.cpp b/Engine/Source/Programs/ReplicationSystemTest/Private/ReplicationSystemTest.cpp
--- a/Engine/Source/Programs/ReplicationSystemTest/Private/ReplicationSystemTest.cpp
+++ b/Engine/Source/Programs/ReplicationSystemTest/Private/ReplicationSystemTest.cpp
@@ -40,10 +40,7 @@ INT32_MAIN_INT32_ARGC_TCHAR_ARGV()
 	PRIVATE_GAllowCommandletRendering = false;
 	PRIVATE_GAllowCommandletAudio = false;
 
-	FCommandLine::Set(
-		TEXT("-nullrhi -log -NoAsyncLoadingThread -NoAsyncPostLoad -noedl -unattended -ENGINEINI=Engine.ini")
-		TEXT(" -ini:Engine:[/Script/Engine.Engine]:AIControllerClassName=/Script/AIModule.AIController")
-	);
+	FCommandLine::Set(TEXT("-nullrhi -log -NoAsyncLoadingThread -NoAsyncPostLoad -noedl -unattended -ENGINEINI=Engine.ini"));
 
 	PreInit();
 	LoadModules();
@@ -120,8 +117,9 @@ static void PreInit()
 
 	FConfigCacheIni::InitializeConfigSystem();
 
-	// 
+	// Config overrides
 	GConfig->SetInt(TEXT("/Script/Engine.GarbageCollectionSettings"), TEXT("gc.MaxObjectsNotConsideredByGC"), 0, GEngineIni);
+	GConfig->SetString(TEXT("/Script/Engine.Engine"), TEXT("AIControllerClassName"), TEXT("/Script/AIModule.AIController"), GEngineIni);
 
 	// Console commands
 	IConsoleManager::Get().ProcessUserConsoleInput(TEXT("Net.IsPushModelEnabled 1"), *GLog, nullptr);
