Create a placeholder default material object in ReplicationSystemTest to prevent a check in the material system.

diff --git a/Engine/Source/Programs/ReplicationSystemTest/Private/ReplicationSystemTest.cpp b/Engine/Source/Programs/ReplicationSystemTest/Private/ReplicationSystemTest.cpp
--- a/Engine/Source/Programs/ReplicationSystemTest/Private/ReplicationSystemTest.cpp
+++ b/Engine/Source/Programs/ReplicationSystemTest/Private/ReplicationSystemTest.cpp
@@ -7,6 +7,7 @@
 #include "AutomationTestRunner.h"
 #include "NullTestRunner.h"
 #include "Net/Core/Trace/Private/NetTraceInternal.h"
+#include "Materials/Material.h"
 
 #if WITH_AUTOMATION_WORKER
 namespace UE::Net
@@ -166,6 +167,10 @@ static void PreInit()
 	// Config overrides
 	GConfig->SetInt(TEXT("/Script/Engine.GarbageCollectionSettings"), TEXT("gc.MaxObjectsNotConsideredByGC"), 0, GEngineIni);
 	GConfig->SetString(TEXT("/Script/Engine.Engine"), TEXT("AIControllerClassName"), TEXT("/Script/AIModule.AIController"), GEngineIni);
+	GConfig->SetString(TEXT("/Script/Engine.Engine"), TEXT("DefaultMaterialName"), TEXT("/Engine/Transient.MockDefaultMaterial"), GEngineIni);
+	GConfig->SetString(TEXT("/Script/Engine.Engine"), TEXT("DefaultLightFunctionMaterialName"), TEXT("/Engine/Transient.MockDefaultMaterial"), GEngineIni);
+	GConfig->SetString(TEXT("/Script/Engine.Engine"), TEXT("DefaultDeferredDecalMaterialName"), TEXT("/Engine/Transient.MockDefaultMaterial"), GEngineIni);
+	GConfig->SetString(TEXT("/Script/Engine.Engine"), TEXT("DefaultPostProcessMaterialName"), TEXT("/Engine/Transient.MockDefaultMaterial"), GEngineIni);
 
 	// Console commands
 	IConsoleManager::Get().ProcessUserConsoleInput(TEXT("Net.IsPushModelEnabled 1"), *GLog, nullptr);
@@ -198,6 +203,9 @@ static void LoadModules()
 	FCoreDelegates::OnInit.Broadcast();
 #endif
 
+	// Create a mock default material to keep the material system happy
+	UMaterial* MockMaterial = NewObject<UMaterial>(GetTransientPackage(), UMaterial::StaticClass(), TEXT("MockDefaultMaterial"), RF_Transient | RF_MarkAsRootSet);
+
 	// ChaosEngineSolvers requires ChaosSolvers. We're not able to call ProcessNewlyLoadedObjects before this module is loaded.
 	FModuleManager::Get().LoadModule(TEXT("ChaosSolvers"));
 	ProcessNewlyLoadedUObjects();
