Removed use of GIsConsoleExecutable in ReplicationSystemTest.

diff --git a/Engine/Source/Programs/ReplicationSystemTest/Private/ReplicationSystemTest.cpp b/Engine/Source/Programs/ReplicationSystemTest/Private/ReplicationSystemTest.cpp
--- a/Engine/Source/Programs/ReplicationSystemTest/Private/ReplicationSystemTest.cpp
+++ b/Engine/Source/Programs/ReplicationSystemTest/Private/ReplicationSystemTest.cpp
@@ -24,8 +24,6 @@ DEFINE_LOG_CATEGORY_STATIC(LogReplicationSystemTest, Log, All);
 
 IMPLEMENT_APPLICATION(ReplicationSystemTest, "");
 
-bool GIsConsoleExecutable = true;
-
 static void PreInit();
 static void LoadModules();
 static void PostInit();
