Iris
- Early out from IrisCoreModule init if iris is compiled in but not enabled on commandline

diff --git a/Engine/Source/Programs/ReplicationSystemTest/Private/ReplicationSystemTest.cpp b/Engine/Source/Programs/ReplicationSystemTest/Private/ReplicationSystemTest.cpp
--- a/Engine/Source/Programs/ReplicationSystemTest/Private/ReplicationSystemTest.cpp
+++ b/Engine/Source/Programs/ReplicationSystemTest/Private/ReplicationSystemTest.cpp
@@ -49,7 +49,7 @@ INT32_MAIN_INT32_ARGC_TCHAR_ARGV()
 
 	// Init cmd line used for test
 	{
-		FString CmdLineOverride(TEXT("-nullrhi -log -NoAsyncLoadingThread -NoAsyncPostLoad -noedl -unattended -ENGINEINI=Engine.ini"));
+		FString CmdLineOverride(TEXT("-nullrhi -log -NoAsyncLoadingThread -NoAsyncPostLoad -noedl -unattended -ENGINEINI=Engine.ini -UseIrisReplication=1"));
 		FString NetTraceVerbosity;
 		if(FParse::Value(ToCStr(OriginalCmdLine), TEXT("NetTrace="), NetTraceVerbosity))
 		{
@@ -173,6 +173,10 @@ static void PreInit()
 	// Console commands
 	IConsoleManager::Get().ProcessUserConsoleInput(TEXT("Net.IsPushModelEnabled 1"), *GLog, nullptr);
 
+	// Console commands
+	IConsoleManager::Get().ProcessUserConsoleInput(TEXT("Log LogIrisBridge Verbose"), *GLog, nullptr);
+
+
 	GGameThreadId = FPlatformTLS::GetCurrentThreadId();
 	FTaskGraphInterface::Startup(FPlatformMisc::NumberOfCores());
 	FTaskGraphInterface::Get().AttachToThread(ENamedThreads::GameThread);
diff --git a/Engine/Source/Runtime/Experimental/Iris/Core/Private/Iris/IrisCoreModule.cpp b/Engine/Source/Runtime/Experimental/Iris/Core/Private/Iris/IrisCoreModule.cpp
--- a/Engine/Source/Runtime/Experimental/Iris/Core/Private/Iris/IrisCoreModule.cpp
+++ b/Engine/Source/Runtime/Experimental/Iris/Core/Private/Iris/IrisCoreModule.cpp
@@ -40,12 +40,15 @@ private:
 		{
 			UE::Net::SetUseIrisReplication(UseIrisReplication > 0);
 		}
-		
-		RegisterPropertyNetSerializerSelectorTypes();
 
-		UE_NET_IRIS_INIT_LEGACY_PUSH_MODEL();
+		if (UE::Net::ShouldUseIrisReplication())
+		{
+			RegisterPropertyNetSerializerSelectorTypes();
+
+			UE_NET_IRIS_INIT_LEGACY_PUSH_MODEL();
 		
-		ModulesChangedHandle = FModuleManager::Get().OnModulesChanged().AddRaw(this, &FIrisCoreModule::OnModulesChanged);
+			ModulesChangedHandle = FModuleManager::Get().OnModulesChanged().AddRaw(this, &FIrisCoreModule::OnModulesChanged);
+		}
 	}
 
 	virtual void ShutdownModule() override
