Removed dependency to deprecated ProfilerService from several standalone programs.

#rnx

diff --git a/Engine/Source/Programs/AutoRTFMEngineTests/AutoRTFMEngineTests.Build.cs b/Engine/Source/Programs/AutoRTFMEngineTests/AutoRTFMEngineTests.Build.cs
--- a/Engine/Source/Programs/AutoRTFMEngineTests/AutoRTFMEngineTests.Build.cs
+++ b/Engine/Source/Programs/AutoRTFMEngineTests/AutoRTFMEngineTests.Build.cs
@@ -12,26 +12,25 @@ public class AutoRTFMEngineTests : ModuleRules
 		PrivateDependencyModuleNames.AddRange(
 			new string[]
 			{
- 				"ApplicationCore",
- 				"AutomationController",
- 				"AutomationWorker",
-                "Core",
+				"ApplicationCore",
+				"AutomationController",
+				"AutomationWorker",
+				"Core",
 				"Projects",
 				"Engine",
 				"HeadMountedDisplay",
 				"InstallBundleManager",
-                "MediaUtils",
+				"MediaUtils",
 				"MRMesh",
 				"MoviePlayer",
 				"MoviePlayerProxy",
 				"MovieScene",
 				"PreLoadScreen",
-				"ProfilerService",
 				"SessionServices",
 				"SlateNullRenderer",
 				"SlateRHIRenderer",
 				"ProfileVisualizer",
- 			}
+			}
 		);
 	}
 }
diff --git a/Engine/Source/Programs/ChaosVisualDebugger/ChaosVisualDebugger.Build.cs b/Engine/Source/Programs/ChaosVisualDebugger/ChaosVisualDebugger.Build.cs
--- a/Engine/Source/Programs/ChaosVisualDebugger/ChaosVisualDebugger.Build.cs
+++ b/Engine/Source/Programs/ChaosVisualDebugger/ChaosVisualDebugger.Build.cs
@@ -30,7 +30,6 @@ public class ChaosVisualDebugger : ModuleRules
 				"SlateCore",
 				"StandaloneRenderer",
 				"TraceLog",
-				"ProfilerService",
 				"ProfileVisualizer",
 				"PropertyAccessEditor"
 			}
diff --git a/Engine/Source/Programs/LiveLinkHub/Source/LiveLinkHubLauncher/LiveLinkHubLauncher.Build.cs b/Engine/Source/Programs/LiveLinkHub/Source/LiveLinkHubLauncher/LiveLinkHubLauncher.Build.cs
--- a/Engine/Source/Programs/LiveLinkHub/Source/LiveLinkHubLauncher/LiveLinkHubLauncher.Build.cs
+++ b/Engine/Source/Programs/LiveLinkHub/Source/LiveLinkHubLauncher/LiveLinkHubLauncher.Build.cs
@@ -36,7 +36,6 @@ public class LiveLinkHubLauncher : ModuleRules
 				"MoviePlayer",
 				"MoviePlayerProxy",
 				"ProfileVisualizer",
-				"ProfilerService",
 				"Projects",
 				"PreLoadScreen",
 				"PIEPreviewDeviceProfileSelector",
diff --git a/Engine/Source/Programs/ReplicationSystemTest/ReplicationSystemTest.Build.cs b/Engine/Source/Programs/ReplicationSystemTest/ReplicationSystemTest.Build.cs
--- a/Engine/Source/Programs/ReplicationSystemTest/ReplicationSystemTest.Build.cs
+++ b/Engine/Source/Programs/ReplicationSystemTest/ReplicationSystemTest.Build.cs
@@ -12,26 +12,25 @@ public class ReplicationSystemTest : ModuleRules
 		PrivateDependencyModuleNames.AddRange(
 			new string[]
 			{
- 				"ApplicationCore",
- 				"AutomationController",
- 				"AutomationWorker",
-                "Core",
+				"ApplicationCore",
+				"AutomationController",
+				"AutomationWorker",
+				"Core",
 				"Projects",
 				"Engine",
 				"HeadMountedDisplay",
 				"InstallBundleManager",
-                "MediaUtils",
+				"MediaUtils",
 				"MRMesh",
 				"MoviePlayer",
 				"MoviePlayerProxy",
 				"PreLoadScreen",
-				"ProfilerService",
 				"ReplicationSystemTestPlugin",
 				"SessionServices",
 				"SlateNullRenderer",
 				"SlateRHIRenderer",
 				"ProfileVisualizer",
- 			}
+			}
 		);
 	}
 }
