Fix build

#rnx

diff --git a/Engine/Source/Programs/AutoRTFMEngineTests/Private/AutoRTFMEngineTests.cpp b/Engine/Source/Programs/AutoRTFMEngineTests/Private/AutoRTFMEngineTests.cpp
--- a/Engine/Source/Programs/AutoRTFMEngineTests/Private/AutoRTFMEngineTests.cpp
+++ b/Engine/Source/Programs/AutoRTFMEngineTests/Private/AutoRTFMEngineTests.cpp
@@ -259,7 +259,7 @@ static void LoadModules()
 	// Always attempt to load CoreUObject. It requires additional pre-init which is called from its module's StartupModule method.
 #if WITH_COREUOBJECT
 	// Always register the UObjects callback for VNI and general consistency with the callbacks ProcessNewlyLoadedUObjects calls.
-	RegisterProcessNewlyLoadedUObjects();
+	RegisterModularObjectsProcessing();
 	FModuleManager::Get().LoadModule(TEXT("CoreUObject"));
 
 	FCoreDelegates::OnInit.Broadcast();
diff --git a/Engine/Source/Programs/ReplicationSystemTest/Private/ReplicationSystemTest.cpp b/Engine/Source/Programs/ReplicationSystemTest/Private/ReplicationSystemTest.cpp
--- a/Engine/Source/Programs/ReplicationSystemTest/Private/ReplicationSystemTest.cpp
+++ b/Engine/Source/Programs/ReplicationSystemTest/Private/ReplicationSystemTest.cpp
@@ -198,7 +198,7 @@ static void LoadModules()
 	// Always attempt to load CoreUObject. It requires additional pre-init which is called from its module's StartupModule method.
 #if WITH_COREUOBJECT
 	// Always register the UObjects callback for VNI and general consistency with the callbacks ProcessNewlyLoadedUObjects calls.
-	RegisterProcessNewlyLoadedUObjects();
+	RegisterModularObjectsProcessing();
 	FModuleManager::Get().LoadModule(TEXT("CoreUObject"));
 
 	FCoreDelegates::OnInit.Broadcast();
