Enable AutoRTFM in ReplicationSystem test programs.

diff --git a/Engine/Source/Programs/ReplicationSystemLowLevelTests/Private/ReplicationSystemLowLevelTests.cpp b/Engine/Source/Programs/ReplicationSystemLowLevelTests/Private/ReplicationSystemLowLevelTests.cpp
--- a/Engine/Source/Programs/ReplicationSystemLowLevelTests/Private/ReplicationSystemLowLevelTests.cpp
+++ b/Engine/Source/Programs/ReplicationSystemLowLevelTests/Private/ReplicationSystemLowLevelTests.cpp
@@ -13,6 +13,7 @@
 #include "SlateGlobals.h"
 #include "Styling/SlateWidgetStyleContainerBase.h"
 #include "TestCommon/Initialization.h"
+#include "AutoRTFM/AutoRTFMUE.h"
 
 #include <catch2/catch_test_macros.hpp>
 
@@ -73,6 +74,7 @@ GROUP_BEFORE_GLOBAL(Catch::DefaultGroup)
 		LOG_SCOPE_VERBOSITY_OVERRIDE(LogUObjectGlobals, ELogVerbosity::Fatal);
 		LOG_SCOPE_VERBOSITY_OVERRIDE(LogStreaming, ELogVerbosity::Error);
 		InitAll(true, true);
+		AutoRTFM::InitializeForUE();
 	}
 
 	for (const FName ModuleToLoad : {"IrisCore", "ReplicationSystemTestPlugin"})
diff --git a/Engine/Source/Programs/ReplicationSystemLowLevelTests/ReplicationSystemLowLevelTests.Target.cs b/Engine/Source/Programs/ReplicationSystemLowLevelTests/ReplicationSystemLowLevelTests.Target.cs
--- a/Engine/Source/Programs/ReplicationSystemLowLevelTests/ReplicationSystemLowLevelTests.Target.cs
+++ b/Engine/Source/Programs/ReplicationSystemLowLevelTests/ReplicationSystemLowLevelTests.Target.cs
@@ -20,6 +20,12 @@ public class ReplicationSystemLowLevelTestsTarget : TestTargetRules
 		bWithPushModel = true;
 
 		bEnableTrace = true;
+
+		if (Target.Platform.IsInGroup(UnrealPlatformGroup.Linux) || Target.Platform == UnrealTargetPlatform.Win64)
+		{
+			bUseAutoRTFMCompiler = true;
+		}
+
 		GlobalDefinitions.Add("UE_NET_TEST_FAKE_REP_TAGS=1");
 		// Load time profiling brings object construction to a crawl.
 		GlobalDefinitions.Add("LOADTIMEPROFILERTRACE_ENABLED=0");
diff --git a/Engine/Source/Programs/ReplicationSystemTest/Private/ReplicationSystemTest.cpp b/Engine/Source/Programs/ReplicationSystemTest/Private/ReplicationSystemTest.cpp
--- a/Engine/Source/Programs/ReplicationSystemTest/Private/ReplicationSystemTest.cpp
+++ b/Engine/Source/Programs/ReplicationSystemTest/Private/ReplicationSystemTest.cpp
@@ -8,6 +8,7 @@
 #include "NullTestRunner.h"
 #include "Net/Core/Trace/Private/NetTraceInternal.h"
 #include "Materials/Material.h"
+#include "AutoRTFM/AutoRTFMUE.h"
 
 #if WITH_AUTOMATION_WORKER
 namespace UE::Net
@@ -128,6 +129,8 @@ static void PreInit()
 	FPlatformApplicationMisc::Init();
 #endif
 	FPlatformMemory::Init();
+	
+	AutoRTFM::InitializeForUE();
 
 	// Initialize trace
 	FString Parameter;
diff --git a/Engine/Source/Programs/ReplicationSystemTest/ReplicationSystemTest.Target.cs b/Engine/Source/Programs/ReplicationSystemTest/ReplicationSystemTest.Target.cs
--- a/Engine/Source/Programs/ReplicationSystemTest/ReplicationSystemTest.Target.cs
+++ b/Engine/Source/Programs/ReplicationSystemTest/ReplicationSystemTest.Target.cs
@@ -33,12 +33,18 @@ public class ReplicationSystemTestTarget : TargetRules
 		bLegalToDistributeBinary = true;
 
 		// Always want logging so we can see the test results
-        bUseLoggingInShipping = true;
+		bUseLoggingInShipping = true;
 
 		// Network config
 		bWithPushModel = true;		
 
 		bEnableTrace = true;
+
+		if (Target.Platform.IsInGroup(UnrealPlatformGroup.Linux) || Target.Platform == UnrealTargetPlatform.Win64)
+		{
+			bUseAutoRTFMCompiler = true;
+		}
+		
 		GlobalDefinitions.Add("UE_NET_TEST_FAKE_REP_TAGS=1");
 		// Load time profiling brings object construction to a crawl.
 		GlobalDefinitions.Add("LOADTIMEPROFILERTRACE_ENABLED=0");
